LiveLinkHub: Follow-up fixes to CL 34680018 / CL 34681546.

Fixes a GenerateProjectFiles warning, and properly re-enables vendor plugins in the existing LiveLinkhub target.

diff --git a/Engine/Programs/LiveLinkHub/Config/DefaultEngine.ini b/Engine/Programs/LiveLinkHub/Config/DefaultEngine.ini
--- a/Engine/Programs/LiveLinkHub/Config/DefaultEngine.ini
+++ b/Engine/Programs/LiveLinkHub/Config/DefaultEngine.ini
@@ -30,3 +30,12 @@ bShareMaterialShaderCode=false
 
 [ConsoleVariables]
 DetailsPanel.Style.EnableCardLayout=true
+
+[Plugins]
+; Optional plugins, will not cause an error if disabled.
++ProgramOptionalPlugins="LiveLinkViconDataStream"
++ProgramOptionalPlugins="Smartsuit"
++ProgramOptionalPlugins="PoseAILiveLink"
+;+ProgramOptionalPlugins="LiveLinkMVNPlugin"
++ProgramOptionalPlugins="MocopiLiveLink"
+;+ProgramOptionalPlugins="OptitrackLiveLink"
diff --git a/Engine/Source/Programs/LiveLinkHubEditor/Source/LiveLinkHubEditor.Target.cs b/Engine/Source/Programs/LiveLinkHubEditor/Source/LiveLinkHubEditor.Target.cs
--- a/Engine/Source/Programs/LiveLinkHubEditor/Source/LiveLinkHubEditor.Target.cs
+++ b/Engine/Source/Programs/LiveLinkHubEditor/Source/LiveLinkHubEditor.Target.cs
@@ -15,6 +15,7 @@ public class LiveLinkHubEditorTarget : TargetRules
 	{
 		Type = TargetType.Editor;
 		bExplicitTargetForType = true;
+		bGenerateProgramProject = true;
 
 		SolutionDirectory = "Programs/LiveLink";
 		LaunchModuleName = "LiveLinkHubLauncher";
@@ -66,10 +67,5 @@ public class LiveLinkHubEditorTarget : TargetRules
 		{
 			OutputFile += ".exe";
 		}
-
-		if (bGenerateProjectFiles)
-		{
-			Type = TargetType.Program;
-		}
 	}
 }
