LiveLinkHub: When running as editor target, set project file to get desired config hierarchy.

Also fixes a mismatched INI property name.

diff --git a/Engine/Programs/LiveLinkHub/Config/DefaultEngine.ini b/Engine/Programs/LiveLinkHub/Config/DefaultEngine.ini
--- a/Engine/Programs/LiveLinkHub/Config/DefaultEngine.ini
+++ b/Engine/Programs/LiveLinkHub/Config/DefaultEngine.ini
@@ -15,7 +15,7 @@ DefaultSourceFrameBufferSize=1
 
 [/Script/LiveLinkInterface.LiveLinkDefaultSubjectSettings]
 bRebroadcastSubjectsByDefault=true
-bAllowModifyingRebroadcast=false
+bAllowEditingRebroadcastProperty=false
 
 [/Script/LiveLinkHub.LiveLinkHubSettings]
 bAutoAddDiscoveredClients=true
diff --git a/Engine/Source/Programs/LiveLinkHub/Private/LiveLinkHubRun.cpp b/Engine/Source/Programs/LiveLinkHub/Private/LiveLinkHubRun.cpp
--- a/Engine/Source/Programs/LiveLinkHub/Private/LiveLinkHubRun.cpp
+++ b/Engine/Source/Programs/LiveLinkHub/Private/LiveLinkHubRun.cpp
@@ -43,6 +43,10 @@ int32 RunLiveLinkHub(const TCHAR* CommandLine)
 		});
 #endif
 
+#if !IS_PROGRAM
+	FPaths::SetProjectFilePath(TEXT("../../Source/Programs/LiveLinkHubEditor/LiveLinkHubEditor.uproject"));
+#endif
+
 	// Must match the path in LiveLinkHubStaging.cs / StageLiveLinkHub.automation.cs
 	const FString RemappedFolder = FPaths::EngineDir() / TEXT("Programs/LiveLinkHub/Staged");
 
diff --git a/Engine/Source/Programs/LiveLinkHubEditor/Config/DefaultEngine.ini b/Engine/Source/Programs/LiveLinkHubEditor/Config/DefaultEngine.ini
--- a/Engine/Source/Programs/LiveLinkHubEditor/Config/DefaultEngine.ini
+++ b/Engine/Source/Programs/LiveLinkHubEditor/Config/DefaultEngine.ini
@@ -15,7 +15,7 @@ DefaultSourceFrameBufferSize=1
 
 [/Script/LiveLinkInterface.LiveLinkDefaultSubjectSettings]
 bRebroadcastSubjectsByDefault=true
-bAllowModifyingRebroadcast=false
+bAllowEditingRebroadcastProperty=false
 
 [/Script/LiveLinkHub.LiveLinkHubSettings]
 bAutoAddDiscoveredClients=true
diff --git a/Engine/Source/Programs/LiveLinkHubEditor/LiveLinkHubEditor.uproject b/Engine/Source/Programs/LiveLinkHubEditor/LiveLinkHubEditor.uproject
--- a/Engine/Source/Programs/LiveLinkHubEditor/LiveLinkHubEditor.uproject
+++ b/Engine/Source/Programs/LiveLinkHubEditor/LiveLinkHubEditor.uproject
@@ -7,5 +7,6 @@
 		"Linux",
 		"Mac",
 		"Windows"
-	]
+	],
+	"DisableEnginePluginsByDefault": true
 }
diff --git a/Engine/Source/Programs/LiveLinkHubEditor/Source/LiveLinkHubEditor.Target.cs b/Engine/Source/Programs/LiveLinkHubEditor/Source/LiveLinkHubEditor.Target.cs
--- a/Engine/Source/Programs/LiveLinkHubEditor/Source/LiveLinkHubEditor.Target.cs
+++ b/Engine/Source/Programs/LiveLinkHubEditor/Source/LiveLinkHubEditor.Target.cs
@@ -1,6 +1,7 @@
 // Copyright Epic Games, Inc. All Rights Reserved.
 
 using EpicGames.Core;
+using UnrealBuildBase;
 using UnrealBuildTool;
 
 
@@ -67,5 +68,27 @@ public class LiveLinkHubEditorTarget : TargetRules
 		{
 			OutputFile += ".exe";
 		}
+
+		// Copy the target receipt into the project binaries directory.
+		// Prevents "Would you like to build the editor?" prompt on startup
+		// when running with project context.
+		DirectoryReference ReceiptSrcDir = Unreal.EngineDirectory;
+		DirectoryReference ReceiptDestDir = DirectoryReference.Combine(
+			Unreal.EngineDirectory, "Source", "Programs", "LiveLinkHubEditor");
+
+		FileReference ReceiptSrcPath = TargetReceipt.GetDefaultPath(ReceiptSrcDir, BaseExeName, Platform, Configuration, Architectures);
+		FileReference ReceiptDestPath = TargetReceipt.GetDefaultPath(ReceiptDestDir, BaseExeName, Platform, Configuration, Architectures);
+
+		PostBuildSteps.Add($"echo Copying \"{ReceiptSrcPath}\" to \"{ReceiptDestPath}\"");
+		DirectoryReference.CreateDirectory(ReceiptDestPath.Directory);
+
+		if (Platform == UnrealTargetPlatform.Win64)
+		{
+			PostBuildSteps.Add($"copy /Y \"{ReceiptSrcPath}\" \"{ReceiptDestPath}\"");
+		}
+		else
+		{
+			PostBuildSteps.Add($"cp -a \"{ReceiptSrcPath}\" \"{ReceiptDestPath}\"");
+		}
 	}
 }
