Splitting the CaptureManager plugin

This work is necessary because LiveLink Hub is now an Editor target, so distinguishing between modules that should be part of the LiveLink Hub and modules that should be part of the UE isn't possible from a uplugin file. And that separation should be done on a plugin level.

diff --git a/Engine/Source/Programs/LiveLinkHub/LiveLinkHub.Target.cs b/Engine/Source/Programs/LiveLinkHub/LiveLinkHub.Target.cs
--- a/Engine/Source/Programs/LiveLinkHub/LiveLinkHub.Target.cs
+++ b/Engine/Source/Programs/LiveLinkHub/LiveLinkHub.Target.cs
@@ -71,7 +71,8 @@ public class LiveLinkHubTarget : TargetRules
 		{
 			OptionalPlugins.AddRange(new string[]
 			{
-				"CaptureManager"
+				"CaptureManagerApp",
+				"CaptureManagerCore"
 			});
 		}
 
diff --git a/Engine/Source/Programs/LiveLinkHubEditor/Source/LiveLinkHubEditor.Target.cs b/Engine/Source/Programs/LiveLinkHubEditor/Source/LiveLinkHubEditor.Target.cs
--- a/Engine/Source/Programs/LiveLinkHubEditor/Source/LiveLinkHubEditor.Target.cs
+++ b/Engine/Source/Programs/LiveLinkHubEditor/Source/LiveLinkHubEditor.Target.cs
@@ -46,7 +46,8 @@ public class LiveLinkHubEditorTarget : TargetRules
 		{
 			OptionalPlugins.AddRange(new string[]
 			{
-				"CaptureManager"
+				"CaptureManagerApp",
+				"CaptureManagerCore"
 			});
 		}
 
