LiveLinkHub - Fix PreProcessors not working in LiveLinkHubEditor target

diff --git a/Engine/Source/Programs/LiveLinkHubEditor/Config/DefaultGame.ini b/Engine/Source/Programs/LiveLinkHubEditor/Config/DefaultGame.ini
--- a/Engine/Source/Programs/LiveLinkHubEditor/Config/DefaultGame.ini
+++ b/Engine/Source/Programs/LiveLinkHubEditor/Config/DefaultGame.ini
@@ -7,3 +7,8 @@ bStagePlatformDirs=false
 bStageUAT=false
 bStagePython=false
 bIsForExternalDistribution=true
+
+[/Script/LiveLink.LiveLinkSettings]
+DefaultSettingsClass=/Script/LiveLinkHub.LiveLinkHubSubjectSettings
+bPreProcessRebroadcastFrames=true
+bTranslateRebroadcastFrames=true
\ No newline at end of file
