Enable XInputDevice to LiveLinkHubEditor target

diff --git a/Engine/Source/Programs/LiveLinkHubEditor/Source/LiveLinkHubEditor.Target.cs b/Engine/Source/Programs/LiveLinkHubEditor/Source/LiveLinkHubEditor.Target.cs
--- a/Engine/Source/Programs/LiveLinkHubEditor/Source/LiveLinkHubEditor.Target.cs
+++ b/Engine/Source/Programs/LiveLinkHubEditor/Source/LiveLinkHubEditor.Target.cs
@@ -46,6 +46,7 @@ public class LiveLinkHubEditorTarget : TargetRules
 			"UdpMessaging",
 			"CameraCalibrationCore",
 			"AppleARKitFaceSupport",
+			"XInputDevice"
 		});
 
 		if (bEnableCaptureManagerPlugin)
