Enable MetaHuman Live Link plugin in LLH editor

diff --git a/Engine/Source/Programs/LiveLinkHubEditor/Source/LiveLinkHubEditor.Target.cs b/Engine/Source/Programs/LiveLinkHubEditor/Source/LiveLinkHubEditor.Target.cs
--- a/Engine/Source/Programs/LiveLinkHubEditor/Source/LiveLinkHubEditor.Target.cs
+++ b/Engine/Source/Programs/LiveLinkHubEditor/Source/LiveLinkHubEditor.Target.cs
@@ -18,6 +18,10 @@ public class LiveLinkHubEditorTarget : TargetRules
 	[CommandLine("-EnableCaptureManagerPlugin=")]
 	public bool bEnableCaptureManagerPlugin = false;
 
+	// Whether to enable building MetaHuman LiveLink plugin.
+	[CommandLine("-EnableMetaHumanLiveLinkPlugin=")]
+	public bool bEnableMetaHumanLiveLinkPlugin = false;
+
 	public LiveLinkHubEditorTarget(TargetInfo Target) : base(Target)
 	{
 		Type = TargetType.Editor;
@@ -66,6 +70,15 @@ public class LiveLinkHubEditorTarget : TargetRules
 			});
 		}
 
+		if (bEnableMetaHumanLiveLinkPlugin)
+		{
+			EnablePlugins.AddRange(new string[]
+			{
+				"MetaHumanLiveLink",
+				"WmfMedia"
+			});
+		}
+
 		if (bSingleModulePlatform)
 		{
 			// Necessary for staging, but avoided otherwise because it dirties
