LiveLinkHub - Expose MessageBus settings

diff --git a/Engine/Source/Programs/LiveLinkHubEditor/Config/DefaultEngine.ini b/Engine/Source/Programs/LiveLinkHubEditor/Config/DefaultEngine.ini
--- a/Engine/Source/Programs/LiveLinkHubEditor/Config/DefaultEngine.ini
+++ b/Engine/Source/Programs/LiveLinkHubEditor/Config/DefaultEngine.ini
@@ -22,6 +22,7 @@ bAllowEditingRebroadcastProperty=false
 bAutoAddDiscoveredClients=true
 +ProjectSettingsToDisplay=LiveLink
 +ProjectSettingsToDisplay=LiveLinkHub
++ProjectSettingsToDisplay=UDPMessaging
 
 [/Script/AppleARKit.AppleARKitSettings]
 bEnableLiveLinkForFaceTracking=true
